Add stub for new Compact Attribute Runtime plugin

diff --git a/Engine/Plugins/Experimental/Animation/PackedAttributes/PackedAttributes.uplugin b/Engine/Plugins/Experimental/Animation/PackedAttributes/PackedAttributes.uplugin
new file mode 100644
--- /dev/null
+++ b/Engine/Plugins/Experimental/Animation/PackedAttributes/PackedAttributes.uplugin
@@ -0,0 +1,32 @@
+{
+	"FileVersion": 3,
+	"Version": 1,
+	"VersionName": "0.1",
+	"FriendlyName": "Packed Attributes",
+	"Description": "",
+	"Category": "Animation",
+	"CreatedBy": "Epic Games, Inc.",
+	"CreatedByURL": "https://epicgames.com",
+	"DocsURL": "",
+	"MarketplaceURL": "",
+	"SupportURL": "",
+	"EnabledByDefault": false,
+	"CanContainContent": true,
+	"IsBetaVersion": false,
+	"IsExperimentalVersion": true,
+	"Installed": false,
+	"Modules": [
+		{
+			"Name": "PackedAttributes",
+			"Type": "Runtime",
+			"LoadingPhase": "Default"
+		},
+		{
+			"Name": "PackedAttributesTestSuite",
+			"Type": "UncookedOnly",
+			"LoadingPhase": "Default"
+		}
+	],
+	"Plugins": [
+	]
+}
diff --git a/Engine/Plugins/Experimental/Animation/PackedAttributes/Source/Runtime/PackedAttributes.Build.cs b/Engine/Plugins/Experimental/Animation/PackedAttributes/Source/Runtime/PackedAttributes.Build.cs
new file mode 100644
--- /dev/null
+++ b/Engine/Plugins/Experimental/Animation/PackedAttributes/Source/Runtime/PackedAttributes.Build.cs
@@ -0,0 +1,24 @@
+// Copyright Epic Games, Inc. All Rights Reserved.
+
+using UnrealBuildTool;
+
+public class PackedAttributes : ModuleRules
+{
+	public PackedAttributes(ReadOnlyTargetRules Target) : base(Target)
+	{
+		PublicDependencyModuleNames.AddRange(
+			new string[]
+			{
+				"Core",
+				"CoreUObject",
+				"Engine",
+			}
+		);
+
+		PrivateDependencyModuleNames.AddRange(
+			new string[]
+			{
+			}
+		);
+	}
+}
diff --git a/Engine/Plugins/Experimental/Animation/PackedAttributes/Source/Runtime/Private/PackedAttributes/Module.cpp b/Engine/Plugins/Experimental/Animation/PackedAttributes/Source/Runtime/Private/PackedAttributes/Module.cpp
new file mode 100644
--- /dev/null
+++ b/Engine/Plugins/Experimental/Animation/PackedAttributes/Source/Runtime/Private/PackedAttributes/Module.cpp
@@ -0,0 +1,20 @@
+// Copyright Epic Games, Inc. All Rights Reserved.
+
+#include "Modules/ModuleManager.h"
+
+namespace UE::Anim::PackedAttributes
+{
+	class FRuntimeModule : public IModuleInterface
+	{
+	public:
+		virtual void StartupModule() override
+		{
+		}
+
+		virtual void ShutdownModule() override
+		{
+		}
+	};
+}
+
+IMPLEMENT_MODULE(UE::Anim::PackedAttributes::FRuntimeModule, PackedAttributes)
diff --git a/Engine/Plugins/Experimental/Animation/PackedAttributes/Source/TestSuite/PackedAttributesTestSuite.Build.cs b/Engine/Plugins/Experimental/Animation/PackedAttributes/Source/TestSuite/PackedAttributesTestSuite.Build.cs
new file mode 100644
--- /dev/null
+++ b/Engine/Plugins/Experimental/Animation/PackedAttributes/Source/TestSuite/PackedAttributesTestSuite.Build.cs
@@ -0,0 +1,32 @@
+// Copyright Epic Games, Inc. All Rights Reserved.
+
+namespace UnrealBuildTool.Rules
+{
+	public class PackedAttributesTestSuite : ModuleRules
+	{
+		public PackedAttributesTestSuite(ReadOnlyTargetRules Target) : base(Target)
+		{
+			PCHUsage = ModuleRules.PCHUsageMode.UseExplicitOrSharedPCHs;
+
+			PublicDependencyModuleNames.AddRange(
+				new string[]
+				{
+					"Core",
+					"CoreUObject",
+					"Engine",
+				}
+			);
+
+			if (Target.bBuildEditor == true)
+			{
+				PrivateDependencyModuleNames.AddRange(
+					new string[]
+					{
+						"UnrealEd",
+						"PackedAttributes",
+					}
+				);
+			}
+		}
+	}
+}
\ No newline at end of file
diff --git a/Engine/Plugins/Experimental/Animation/PackedAttributes/Source/TestSuite/Private/Module.cpp b/Engine/Plugins/Experimental/Animation/PackedAttributes/Source/TestSuite/Private/Module.cpp
new file mode 100644
--- /dev/null
+++ b/Engine/Plugins/Experimental/Animation/PackedAttributes/Source/TestSuite/Private/Module.cpp
@@ -0,0 +1,20 @@
+// Copyright Epic Games, Inc. All Rights Reserved.
+
+#include "Modules/ModuleManager.h"
+
+namespace UE::Anim::PackedAttributes
+{
+	class FTestSuiteModule : public IModuleInterface
+	{
+	public:
+		virtual void StartupModule() override
+		{
+		}
+
+		virtual void ShutdownModule() override
+		{
+		}
+	};
+}
+
+IMPLEMENT_MODULE(UE::Anim::PackedAttributes::FTestSuiteModule, PackedAttributesTestSuite)
